add comment + a bit cleaner function naming for isosurface stuffing edge ID / grid ID code

diff --git a/Engine/Plugins/Experimental/TetMeshing/Source/TetMeshing/Public/Generate/IsosurfaceStuffing.h b/Engine/Plugins/Experimental/TetMeshing/Source/TetMeshing/Public/Generate/IsosurfaceStuffing.h
--- a/Engine/Plugins/Experimental/TetMeshing/Source/TetMeshing/Public/Generate/IsosurfaceStuffing.h
+++ b/Engine/Plugins/Experimental/TetMeshing/Source/TetMeshing/Public/Generate/IsosurfaceStuffing.h
@@ -182,13 +182,43 @@ protected:
 		return TVec3(0, 0, 0);
 	}
 
+	/// Edge and Point ID Packing Scheme:
+	/// We use a bit-packing scheme to generate unique int64 IDs for grid points and canonical BCC edges.
+	/// This is adapted from the scheme used by MarchingCubes.h, extended to handle the BCC grid's diagonals and center points.
+	/// 
+	/// Grid points are defined by an X,Y,Z and a bool indicating if they are centered.
+	/// We assume the grid dimensions cannot exceed 2^15 on any axis, so we can pack grid points
+	/// in the first 49 bits of an int64: 16 bits per X, Y, Z, and 1 bit for the bIsCenter flag.
+	/// Edge IDs are given by one of the grid points bitwise or'd with a 'Tag' indicating the attached edge.
+	/// To make sure each edge only has one possible ID, by convention, the grid point for an edge is chosen by these rules:
+	///  - If the edge is a 'long' BCC edge, along a major axis, use the grid point with the lowest coordinates.
+	///  - Otherwise, for 'short' BCC edges, connecting centered to non-centered points, always use the centered point.
+	/// Edge tags are numbered as follows:
+	///  - 1,2,3 for the long edges along the X,Y,Z axes respectively, as returned by GetMajorAxisEdgeTag.
+	///  - 4-11 for the short edges, as returned by GetCornerEdgeTag(Difference in coordinates).
+	///  (Note that the coordinate difference from center->corner is always 0 or 1 in each dimension.)
+	/// The tag number is bit-shifted by 50 to safely combine with the Point ID.
+	/// 
+	/// In sum, the bits of an Edge ID are laid out as:
+	/// [Bits 53-50: 'Tag' indicating which edge] [Bits 48-0: Point ID for a grid point attached to the edge]
+	/// And Point ID is laid out as:
+	/// [Bit 48: bIsCenter flag] [Bits 47-32: Z coord] [Bits 31-16: Y coord] [Bits 15-0: X coord]
+	
+	static constexpr int64 GetMajorAxisEdgeTag(int32 Axis)
+	{
+		return int64(Axis + 1) << 50;
+	}
+	static constexpr int64 GetCornerEdgeTag(FIntVector3 Delta)
+	{
+		return int64(Delta.X + Delta.Y * 2 + Delta.Z * 4 + 4) << 50;
+	}
 
 	// Return a unique ID for the BCC edge connecting A to B (note: assumes A,B are connected in the BCC lattice)
 	static int64 GetEdgeID(FGridPt A, FGridPt B)
 	{
-		constexpr int64 EDGE_X = int64(1) << 50;
-		constexpr int64 EDGE_Y = int64(2) << 50;
-		constexpr int64 EDGE_Z = int64(3) << 50;
+		constexpr int64 EDGE_X = GetMajorAxisEdgeTag(0);
+		constexpr int64 EDGE_Y = GetMajorAxisEdgeTag(1);
+		constexpr int64 EDGE_Z = GetMajorAxisEdgeTag(2);
 		if (A.bIsCenter == B.bIsCenter)
 		{
 			if (A.Idx.X != B.Idx.X)
@@ -218,16 +248,11 @@ protected:
 				CenterIdx = A.Idx;
 			}
 			checkSlow(Delta.X >= 0 && Delta.X <= 1 && Delta.Y >= 0 && Delta.Y <= 1 && Delta.Z >= 0 && Delta.Z <= 1);
-			int64 CID = GetCornerID(Delta);
-			return GetGridID(CenterIdx, true) | CID;
+			int64 CornerTag = GetCornerEdgeTag(Delta);
+			return GetGridID(CenterIdx, true) | CornerTag;
 		}
 	}
 
-	static inline int64 GetCornerID(FIntVector3 Delta)
-	{
-		return int64(Delta.X + Delta.Y * 2 + Delta.Z * 4 + 4) << 50;
-	}
-
 	static inline int64 GetGridID(const FGridPt& Pt)
 	{
 		return GetGridID(Pt.Idx, Pt.bIsCenter);
@@ -244,9 +269,9 @@ protected:
 	static void GetEdges(FGridPt Pt, int64 EdgeIDsOut[14], FGridPt OtherPtOut[14])
 	{
 		// first six == along major axes, last 8 == corners
-		constexpr int64 EDGE_X = int64(1) << 50;
-		constexpr int64 EDGE_Y = int64(2) << 50;
-		constexpr int64 EDGE_Z = int64(3) << 50;
+		constexpr int64 EDGE_X = GetMajorAxisEdgeTag(0);
+		constexpr int64 EDGE_Y = GetMajorAxisEdgeTag(1);
+		constexpr int64 EDGE_Z = GetMajorAxisEdgeTag(2);
 		int64 PtID = GetGridID(Pt);
 		OtherPtOut[0] = FGridPt(Pt.Idx.X + 1, Pt.Idx.Y, Pt.Idx.Z, Pt.bIsCenter);
 		EdgeIDsOut[0] = PtID | EDGE_X;
@@ -270,7 +295,7 @@ protected:
 					for (int32 ZOff = 0; ZOff < 2; ++ZOff)
 					{
 						OtherPtOut[OutPtIdx] = FGridPt(Pt.Idx.X + XOff, Pt.Idx.Y + YOff, Pt.Idx.Z + ZOff, false);
-						EdgeIDsOut[OutPtIdx] = PtID | GetCornerID(FIntVector3(XOff, YOff, ZOff));
+						EdgeIDsOut[OutPtIdx] = PtID | GetCornerEdgeTag(FIntVector3(XOff, YOff, ZOff));
 						OutPtIdx++;
 					}
 				}
@@ -288,7 +313,7 @@ protected:
 					for (int32 ZOff = 0; ZOff < 2; ++ZOff)
 					{
 						OtherPtOut[OutPtIdx] = FGridPt(Pt.Idx.X - XOff, Pt.Idx.Y - YOff, Pt.Idx.Z - ZOff, true);
-						EdgeIDsOut[OutPtIdx] = GetGridID(OtherPtOut[OutPtIdx]) | GetCornerID(FIntVector3(XOff, YOff, ZOff));
+						EdgeIDsOut[OutPtIdx] = GetGridID(OtherPtOut[OutPtIdx]) | GetCornerEdgeTag(FIntVector3(XOff, YOff, ZOff));
 						checkSlow(EdgeIDsOut[OutPtIdx] == GetEdgeID(Pt, OtherPtOut[OutPtIdx]));
 						OutPtIdx++;
 					}
