minor isosurface stuffing fixes + added test coverage of offset surface generation

diff --git a/Engine/Plugins/Experimental/TetMeshing/Source/TetMeshing/Public/Generate/IsosurfaceStuffing.h b/Engine/Plugins/Experimental/TetMeshing/Source/TetMeshing/Public/Generate/IsosurfaceStuffing.h
--- a/Engine/Plugins/Experimental/TetMeshing/Source/TetMeshing/Public/Generate/IsosurfaceStuffing.h
+++ b/Engine/Plugins/Experimental/TetMeshing/Source/TetMeshing/Public/Generate/IsosurfaceStuffing.h
@@ -111,6 +111,13 @@ public:
 		return *this;
 	}
 
+	virtual void Reset() override
+	{
+		TTetMeshGenerator<RealType>::Reset();
+		ComputedGridPts.Reset();
+		ComputedCutPts.Reset();
+	}
+
 protected:
 
 	struct FGridPt
@@ -154,7 +161,7 @@ protected:
 	{
 		if constexpr (RootMethod == ERootFindingMethod::Bisection)
 		{
-			checkSlow(FMath::Sign(AVal) != FMath::Sign(BVal)); // TODO: does not apply if multi-class?
+			checkSlow(FMath::Sign(AVal - IsoValue) != FMath::Sign(BVal - IsoValue)); // TODO: does not apply if multi-class?
 			if (AVal > BVal)
 			{
 				Swap(AVal, BVal);
diff --git a/Engine/Plugins/Experimental/TetMeshing/Source/TetMeshing/Public/Generate/TetMeshGenerator.h b/Engine/Plugins/Experimental/TetMeshing/Source/TetMeshing/Public/Generate/TetMeshGenerator.h
--- a/Engine/Plugins/Experimental/TetMeshing/Source/TetMeshing/Public/Generate/TetMeshGenerator.h
+++ b/Engine/Plugins/Experimental/TetMeshing/Source/TetMeshing/Public/Generate/TetMeshGenerator.h
@@ -52,7 +52,7 @@ public:
 
 
 	/** clear arrays so that Generate() can be run again */
-	void Reset()
+	virtual void Reset()
 	{
 		Vertices.Reset();
 		Tets.Reset();
