fix CIS issue: PVS Studio warnings in IsosurfaceStuffing.h

diff --git a/Engine/Plugins/Experimental/TetMeshing/Source/TetMeshing/Public/Generate/IsosurfaceStuffing.h b/Engine/Plugins/Experimental/TetMeshing/Source/TetMeshing/Public/Generate/IsosurfaceStuffing.h
--- a/Engine/Plugins/Experimental/TetMeshing/Source/TetMeshing/Public/Generate/IsosurfaceStuffing.h
+++ b/Engine/Plugins/Experimental/TetMeshing/Source/TetMeshing/Public/Generate/IsosurfaceStuffing.h
@@ -211,13 +211,16 @@ protected:
 	/// And Point ID is laid out as:
 	/// [Bit 48: bIsCenter flag] [Bits 47-32: Z coord] [Bits 31-16: Y coord] [Bits 15-0: X coord]
 	
-	static constexpr int64 GetMajorAxisEdgeTag(int32 Axis)
+	// Axis should be 0, 1, or 2 for X, Y, or Z
+	static constexpr int64 GetMajorAxisEdgeTag(int64 Axis)
 	{
 		return int64(Axis + 1) << 50;
 	}
+	// Delta should be 0 or 1 for each coordinate
 	static constexpr int64 GetCornerEdgeTag(FIntVector3 Delta)
 	{
-		return int64(Delta.X + Delta.Y * 2 + Delta.Z * 4 + 4) << 50;
+		int32 CornerEncoded = Delta.X + Delta.Y * 2 + Delta.Z * 4 + 4;
+		return int64(CornerEncoded) << 50;
 	}
 
 	// Return a unique ID for the BCC edge connecting A to B (note: assumes A,B are connected in the BCC lattice)
@@ -399,9 +402,10 @@ protected:
 			RealType SmallestAlpha = 2; // Note: Intentionally larger than max valid Alpha value of 1
 			int32 WarpEdgeIdx = INDEX_NONE;
 			TVec3 WarpPos;
-			auto EvalEdgeRange = [&/*TODO FIX*/](int32 MinEdgeIdx, int32 MaxEdgeIdx, RealType AlphaThreshold)
+			auto EvalEdgeRange = [this, &EdgeIDs, &OtherPts, Pt, PosOut, LabelA, &CutPtIDs, &HasCutPtNum](int32 MinEdgeIdx, int32 MaxEdgeIdx, RealType AlphaThreshold,
+				RealType& SmallestAlphaRef, int32& WarpEdgeIdxRef, TVec3& WarpPosRef)
 			{
-				RealType KeepCutAlphaThreshold = WarpEdgeIdx == INDEX_NONE ? AlphaThreshold : (RealType)2;
+				RealType KeepCutAlphaThreshold = WarpEdgeIdxRef == INDEX_NONE ? AlphaThreshold : (RealType)2;
 				for (int32 EdgeIdx = MinEdgeIdx; EdgeIdx < MaxEdgeIdx; ++EdgeIdx)
 				{
 					TVec3 OtherPos;
@@ -415,21 +419,22 @@ protected:
 						if (Alpha < AlphaThreshold)
 						{
 							KeepCutAlphaThreshold = 2; // stop keeping cut points
-							if (Alpha < SmallestAlpha)
+							if (Alpha < SmallestAlphaRef)
 							{
-								WarpEdgeIdx = EdgeIdx;
-								SmallestAlpha = Alpha;
-								WarpPos = CutData.CutPos;
+								WarpEdgeIdxRef = EdgeIdx;
+								SmallestAlphaRef = Alpha;
+								WarpPosRef = CutData.CutPos;
 							}
 						}
 					}
 				}
 			};
-			EvalEdgeRange(0, 6, AlphaLong); // Process the long edges
-			EvalEdgeRange(6, 14, AlphaShort); // Process the short edges
+			EvalEdgeRange(0, 6, AlphaLong, SmallestAlpha, WarpEdgeIdx, WarpPos); // Process the long edges
+			EvalEdgeRange(6, 14, AlphaShort, SmallestAlpha, WarpEdgeIdx, WarpPos); // Process the short edges
 			ToReturn.bWarpProcessed = true;
 			if (WarpEdgeIdx != INDEX_NONE)
 			{
+				PosOut = WarpPos;
 				for (int32 CutIDIdx = 0; CutIDIdx < HasCutPtNum; ++CutIDIdx)
 				{
 					ComputedCutPts.Remove(CutPtIDs[CutIDIdx]);
