CIS fix for PVS studio warning in IsosurfaceStuffing.h
Pass more variables by reference explicitly rather than by capture because PVS studio appears to not understand lambda captures

diff --git a/Engine/Plugins/Experimental/TetMeshing/Source/TetMeshing/Public/Generate/IsosurfaceStuffing.h b/Engine/Plugins/Experimental/TetMeshing/Source/TetMeshing/Public/Generate/IsosurfaceStuffing.h
--- a/Engine/Plugins/Experimental/TetMeshing/Source/TetMeshing/Public/Generate/IsosurfaceStuffing.h
+++ b/Engine/Plugins/Experimental/TetMeshing/Source/TetMeshing/Public/Generate/IsosurfaceStuffing.h
@@ -402,8 +402,9 @@ protected:
 			RealType SmallestAlpha = 2; // Note: Intentionally larger than max valid Alpha value of 1
 			int32 WarpEdgeIdx = INDEX_NONE;
 			TVec3 WarpPos;
-			auto EvalEdgeRange = [this, &EdgeIDs, &OtherPts, Pt, PosOut, LabelA, &CutPtIDs, &HasCutPtNum](int32 MinEdgeIdx, int32 MaxEdgeIdx, RealType AlphaThreshold,
-				RealType& SmallestAlphaRef, int32& WarpEdgeIdxRef, TVec3& WarpPosRef)
+			auto EvalEdgeRange = [this, Pt, PosOut, LabelA, &EdgeIDs, &OtherPts](int32 MinEdgeIdx, int32 MaxEdgeIdx, RealType AlphaThreshold,
+				// Note: The following are passed by reference explicitly, rather than via capture, because PVS static analysis didn't believe in the captures:
+				RealType& SmallestAlphaRef, int32& WarpEdgeIdxRef, TVec3& WarpPosRef, int32& HasCutPtNumRef, int64* CutPtIDsArrPtr)
 			{
 				RealType KeepCutAlphaThreshold = WarpEdgeIdxRef == INDEX_NONE ? AlphaThreshold : (RealType)2;
 				for (int32 EdgeIdx = MinEdgeIdx; EdgeIdx < MaxEdgeIdx; ++EdgeIdx)
@@ -415,7 +416,7 @@ protected:
 					bool bHasCut = GetCutData(EdgeIDs[EdgeIdx], Pt, OtherPts[EdgeIdx], PosOut, OtherPos, LabelA, OtherData.Label, KeepCutAlphaThreshold, CutData, Alpha);
 					if (bHasCut)
 					{
-						CutPtIDs[HasCutPtNum++] = EdgeIDs[EdgeIdx]; // TODO: make GetCutData report whether there is actually anything in the has, and only add to this array in that case
+						CutPtIDsArrPtr[HasCutPtNumRef++] = EdgeIDs[EdgeIdx]; // TODO: make GetCutData report whether there is actually anything in the has, and only add to this array in that case
 						if (Alpha < AlphaThreshold)
 						{
 							KeepCutAlphaThreshold = 2; // stop keeping cut points
@@ -429,8 +430,8 @@ protected:
 					}
 				}
 			};
-			EvalEdgeRange(0, 6, AlphaLong, SmallestAlpha, WarpEdgeIdx, WarpPos); // Process the long edges
-			EvalEdgeRange(6, 14, AlphaShort, SmallestAlpha, WarpEdgeIdx, WarpPos); // Process the short edges
+			EvalEdgeRange(0, 6, AlphaLong, SmallestAlpha, WarpEdgeIdx, WarpPos, HasCutPtNum, CutPtIDs); // Process the long edges
+			EvalEdgeRange(6, 14, AlphaShort, SmallestAlpha, WarpEdgeIdx, WarpPos, HasCutPtNum, CutPtIDs); // Process the short edges
 			ToReturn.bWarpProcessed = true;
 			if (WarpEdgeIdx != INDEX_NONE)
 			{
